Clamp ColumnWidth from 0.0 to 1.0
* Initially came up as a bug because an Editor Utility Widget was built with a value outside of this.
* Tracing it - UDetailsView::BuildContentWidget will eventually read this and create an SDetailsView with the value and call ColumnSizeData.SetValueColumnWidth which will ensure a max of 1.0.

diff --git a/Engine/Source/Editor/ScriptableEditorWidgets/Public/Components/DetailsView.h b/Engine/Source/Editor/ScriptableEditorWidgets/Public/Components/DetailsView.h
--- a/Engine/Source/Editor/ScriptableEditorWidgets/Public/Components/DetailsView.h
+++ b/Engine/Source/Editor/ScriptableEditorWidgets/Public/Components/DetailsView.h
@@ -49,7 +49,7 @@ public:
 	bool bShowAnimatedPropertiesOption = true;
 
 	/** The default column width */
-	UPROPERTY(EditAnywhere, Category = "View")
+	UPROPERTY(EditAnywhere, Category = "View", meta=(ClampMin="0.0", ClampMax="1.0"))
 	float ColumnWidth = 0.65f;
 
 	/** If false, the details panel's scrollbar will always be hidden. Useful when embedding details panels in widgets that either grow to accommodate them, or with scrollbars of their own. */
