Editor Utility Widget: fix single property view not triggering OnPropertyChange when children of a property change.

diff --git a/Engine/Source/Editor/ScriptableEditorWidgets/Private/Components/SinglePropertyView.cpp b/Engine/Source/Editor/ScriptableEditorWidgets/Private/Components/SinglePropertyView.cpp
--- a/Engine/Source/Editor/ScriptableEditorWidgets/Private/Components/SinglePropertyView.cpp
+++ b/Engine/Source/Editor/ScriptableEditorWidgets/Private/Components/SinglePropertyView.cpp
@@ -10,6 +10,7 @@
 #include "Misc/Attribute.h"
 #include "Modules/ModuleManager.h"
 #include "PropertyEditorModule.h"
+#include "PropertyHandle.h"
 #include "UObject/Class.h"
 #include "UObject/Field.h"
 #include "UObject/Object.h"
@@ -90,9 +91,12 @@ void USinglePropertyView::BuildContentWidget()
 
 				if (SinglePropertyViewWidget.IsValid())
 				{
-					FSimpleDelegate PropertyChanged = FSimpleDelegate::CreateUObject(this, &USinglePropertyView::InternalSinglePropertyChanged);
-					SinglePropertyViewWidget->SetOnPropertyValueChanged(PropertyChanged);
-
+					if (const TSharedPtr<IPropertyHandle> Handle = SinglePropertyViewWidget->GetPropertyHandle())
+					{
+						FSimpleDelegate PropertyChanged = FSimpleDelegate::CreateUObject(this, &USinglePropertyView::InternalSinglePropertyChanged);
+						Handle->SetOnPropertyValueChanged(PropertyChanged);
+						Handle->SetOnChildPropertyValueChanged(PropertyChanged);
+					}
 					GetDisplayWidget()->SetContent(SinglePropertyViewWidget.ToSharedRef());
 					bCreateMissingWidget = false;
 				}
