PR #13006: Improve extensibility of UDetailsView and SinglePropertyView

diff --git a/Engine/Source/Editor/ScriptableEditorWidgets/Public/Components/DetailsView.h b/Engine/Source/Editor/ScriptableEditorWidgets/Public/Components/DetailsView.h
--- a/Engine/Source/Editor/ScriptableEditorWidgets/Public/Components/DetailsView.h
+++ b/Engine/Source/Editor/ScriptableEditorWidgets/Public/Components/DetailsView.h
@@ -73,7 +73,7 @@ public:
 	TArray<FName> PropertiesToShow;
 
 	//~ UPropertyViewBase interface
-private:
+protected:
 	virtual void BuildContentWidget() override;
 	virtual void OnObjectChanged() override;
 	//~ End of UPropertyViewBase interface
diff --git a/Engine/Source/Editor/ScriptableEditorWidgets/Public/Components/SinglePropertyView.h b/Engine/Source/Editor/ScriptableEditorWidgets/Public/Components/SinglePropertyView.h
--- a/Engine/Source/Editor/ScriptableEditorWidgets/Public/Components/SinglePropertyView.h
+++ b/Engine/Source/Editor/ScriptableEditorWidgets/Public/Components/SinglePropertyView.h
@@ -49,7 +49,7 @@ private:
 	void InternalSinglePropertyChanged();
 
 	//~ UPropertyViewBase interface
-private:
+protected:
 	virtual void BuildContentWidget() override;
 	virtual void OnObjectChanged() override;
 	//~ End of UPropertyViewBase interface
